HDRIBackdrop plugin. Missing style code.

diff --git a/Engine/Plugins/Runtime/HDRIBackdrop/Source/HDRIBackdrop/Private/HDRIBackdropStyle.cpp b/Engine/Plugins/Runtime/HDRIBackdrop/Source/HDRIBackdrop/Private/HDRIBackdropStyle.cpp
new file mode 100644
--- /dev/null
+++ b/Engine/Plugins/Runtime/HDRIBackdrop/Source/HDRIBackdrop/Private/HDRIBackdropStyle.cpp
@@ -0,0 +1,37 @@
+// Copyright 1998-2019 Epic Games, Inc. All Rights Reserved.
+
+#include "HDRIBackdropStyle.h"
+#include "Styling/SlateStyle.h"
+#include "Styling/SlateStyleRegistry.h"
+#include "Styling/SlateTypes.h"
+#include "Interfaces/IPluginManager.h"
+
+TSharedPtr<FSlateStyleSet> FHDRIBackdropStyle::StyleSet;
+
+void FHDRIBackdropStyle::Initialize()
+{
+	if (StyleSet.IsValid())
+	{
+		return;
+	}
+
+	StyleSet = MakeShared<FSlateStyleSet>(FName("HDRIBackdropStyle"));
+
+	StyleSet->SetContentRoot(FPaths::EngineContentDir() / TEXT("Editor/Slate"));
+	StyleSet->SetCoreContentRoot(FPaths::EngineContentDir() / TEXT("Slate"));
+
+	FString IconPath = IPluginManager::Get().FindPlugin(TEXT("HDRIBackdrop"))->GetBaseDir() + TEXT("/Resources/HDRIBackdrop.png");
+	StyleSet->Set("HDRIBackdrop.ModesIcon", new FSlateImageBrush(IconPath, FVector2D(40.0f, 40.0f)));
+
+	FSlateStyleRegistry::RegisterSlateStyle(*StyleSet.Get());
+}
+
+void FHDRIBackdropStyle::Shutdown()
+{
+	if (StyleSet.IsValid())
+	{
+		FSlateStyleRegistry::UnRegisterSlateStyle(*StyleSet.Get());
+		ensure(StyleSet.IsUnique());
+		StyleSet.Reset();
+	}
+}
\ No newline at end of file
diff --git a/Engine/Plugins/Runtime/HDRIBackdrop/Source/HDRIBackdrop/Public/HDRIBackdropStyle.h b/Engine/Plugins/Runtime/HDRIBackdrop/Source/HDRIBackdrop/Public/HDRIBackdropStyle.h
new file mode 100644
--- /dev/null
+++ b/Engine/Plugins/Runtime/HDRIBackdrop/Source/HDRIBackdrop/Public/HDRIBackdropStyle.h
@@ -0,0 +1,17 @@
+// Copyright 1998-2019 Epic Games, Inc. All Rights Reserved.
+
+#pragma once
+
+#include "CoreMinimal.h"
+#include "Styling/SlateStyle.h"
+
+class FHDRIBackdropStyle
+{
+public:
+	static void Initialize();
+	static void Shutdown();
+	static TSharedPtr<ISlateStyle> Get() { return StyleSet; }
+
+private:
+	static TSharedPtr<FSlateStyleSet> StyleSet;
+};
\ No newline at end of file
